Fix for incorrect category name for Avid DNx encoding

diff --git a/Engine/Plugins/Media/AvidDNxMedia/Source/Source/Private/AvidDNxEncoderProtocol.h b/Engine/Plugins/Media/AvidDNxMedia/Source/Source/Private/AvidDNxEncoderProtocol.h
--- a/Engine/Plugins/Media/AvidDNxMedia/Source/Source/Private/AvidDNxEncoderProtocol.h
+++ b/Engine/Plugins/Media/AvidDNxMedia/Source/Source/Private/AvidDNxEncoderProtocol.h
@@ -32,10 +32,10 @@ public:
 	virtual void ProcessFrame(FCapturedFrameData InFrame) override;
 
 public:
-	UPROPERTY(config, EditAnywhere, Category = "Avid DNx Settings")
+	UPROPERTY(config, EditAnywhere, Category = "Avid DNX Settings")
 	bool bUncompressed;
 
-	UPROPERTY(config, EditAnywhere, Category = "Avid DNx Settings", meta = (ClampMin = "1", UIMin = "1", ClampMax = "64", UIMax = "64"))
+	UPROPERTY(config, EditAnywhere, Category = "Avid DNX Settings", meta = (ClampMin = "1", UIMin = "1", ClampMax = "64", UIMax = "64"))
 	uint8 NumberOfEncodingThreads;
 
 public:
