Added missing #pragma once.

diff --git a/Engine/Plugins/Media/AvidDNxMedia/Source/Source/Private/MoviePipelineAvidDNxOutput.h b/Engine/Plugins/Media/AvidDNxMedia/Source/Source/Private/MoviePipelineAvidDNxOutput.h
--- a/Engine/Plugins/Media/AvidDNxMedia/Source/Source/Private/MoviePipelineAvidDNxOutput.h
+++ b/Engine/Plugins/Media/AvidDNxMedia/Source/Source/Private/MoviePipelineAvidDNxOutput.h
@@ -1,5 +1,7 @@
 // Copyright Epic Games, Inc. All Rights Reserved.
 
+#pragma once
+
 #include "MoviePipelineVideoOutputBase.h"
 #include "AvidDNxEncoder/AvidDNxEncoder.h"
 #include "MoviePipelineAvidDNxOutput.generated.h"
