Consistent pathing for Runtime/Launch/Resources/Version.h

#rnx

diff --git a/Engine/Platforms/Hololens/Source/Runtime/Core/Private/HoloLensPlatformCrashContext.cpp b/Engine/Platforms/Hololens/Source/Runtime/Core/Private/HoloLensPlatformCrashContext.cpp
--- a/Engine/Platforms/Hololens/Source/Runtime/Core/Private/HoloLensPlatformCrashContext.cpp
+++ b/Engine/Platforms/Hololens/Source/Runtime/Core/Private/HoloLensPlatformCrashContext.cpp
@@ -8,7 +8,7 @@
 #include "HAL/PlatformMallocCrash.h"
 #include "HAL/ExceptionHandling.h"
 #include "HoloLensPlatformCrashContext.h"
-#include "../../Launch/Resources/Version.h"
+#include "Runtime/Launch/Resources/Version.h"
 
 #include "Microsoft/AllowMicrosoftPlatformTypes.h"
 
diff --git a/Engine/Plugins/Media/AvidDNxMedia/Source/Source/Private/AvidDNxEncoder/AvidDNxEncoder.cpp b/Engine/Plugins/Media/AvidDNxMedia/Source/Source/Private/AvidDNxEncoder/AvidDNxEncoder.cpp
--- a/Engine/Plugins/Media/AvidDNxMedia/Source/Source/Private/AvidDNxEncoder/AvidDNxEncoder.cpp
+++ b/Engine/Plugins/Media/AvidDNxMedia/Source/Source/Private/AvidDNxEncoder/AvidDNxEncoder.cpp
@@ -1,9 +1,9 @@
 // Copyright Epic Games, Inc. All Rights Reserved.
 
 #include "AvidDNxEncoder/AvidDNxEncoder.h"
-#include "Launch/Resources/Version.h"
 #include "AvidDNxMediaModule.h"
 #include "HAL/PlatformTime.h"
+#include "Runtime/Launch/Resources/Version.h"
 
 THIRD_PARTY_INCLUDES_START
 #include <AvidDNxCodec.h>
diff --git a/Engine/Plugins/Media/AvidDNxMedia/Source/Source/Private/AvidDNxEncoderProtocol.cpp b/Engine/Plugins/Media/AvidDNxMedia/Source/Source/Private/AvidDNxEncoderProtocol.cpp
--- a/Engine/Plugins/Media/AvidDNxMedia/Source/Source/Private/AvidDNxEncoderProtocol.cpp
+++ b/Engine/Plugins/Media/AvidDNxMedia/Source/Source/Private/AvidDNxEncoderProtocol.cpp
@@ -4,8 +4,8 @@
 
 #include "AvidDNxMediaModule.h"
 #include "Engine/RendererSettings.h"
-#include "Launch/Resources/Version.h"
 #include "Misc/CoreMiscDefines.h"
+#include "Runtime/Launch/Resources/Version.h"
 
 THIRD_PARTY_INCLUDES_START
 #include <AvidDNxCodec.h>
