MoviePipeline: Added additional quality options (LB 8-bit, SQ 8-bit, HQX 10-bit, and RGB 4:4:4 12-bit) to the graph's Avid DNxHR Movie node.

diff --git a/Engine/Plugins/Media/AvidDNxMedia/Source/Source/Private/AvidDNxEncoder/AvidDNxEncoder.cpp b/Engine/Plugins/Media/AvidDNxMedia/Source/Source/Private/AvidDNxEncoder/AvidDNxEncoder.cpp
--- a/Engine/Plugins/Media/AvidDNxMedia/Source/Source/Private/AvidDNxEncoder/AvidDNxEncoder.cpp
+++ b/Engine/Plugins/Media/AvidDNxMedia/Source/Source/Private/AvidDNxEncoder/AvidDNxEncoder.cpp
@@ -1,6 +1,8 @@
 // Copyright Epic Games, Inc. All Rights Reserved.
 
 #include "AvidDNxEncoder/AvidDNxEncoder.h"
+
+#include "Async/ParallelFor.h"
 #include "AvidDNxMediaModule.h"
 #include "HAL/PlatformTime.h"
 #include "Runtime/Launch/Resources/Version.h"
@@ -17,21 +19,44 @@ namespace AvidDNx
 	const wchar_t* ProductUID = L"06.9d.41.48.a0.cb.48.d4.af.19.54.da.bd.09.2a.9f";
 
 	/**
-	 * Creates one sub-sampled pixel stored in Y0CbY1Cr from two input BGRA pixels.
+	 * Converts two RGB pixels (InputColor0 and InputColor1) to sub-sampled YCbCr Rec. 709 and video range. Supports both 8-bit and 16-bit color.
 	 *
-	 * @note The input colors will be transformed to 709 and video range at 8 bit.
+	 * InputType supports FColor and FFloat16Color.
+	 * ColorComponentType supports uint8 and unsigned short.
+	 * ColorContainer supports FY0CbY1Cr and FY0CbY1Cr_16bit.
 	 */
-	FAvidDNxEncoder::FY0CbY1Cr SubSample422(const FColor* InBgra0, const FColor* InBgra1)
+	template<typename InputType, typename ColorComponentType, typename ColorContainer>
+	ColorContainer RBGtoYCbCrRec709(const InputType* InputColor0, const InputType* InputColor1)
 	{
-		const float R0 = InBgra0->R / 255.0f;
-		const float G0 = InBgra0->G / 255.0f;
-		const float B0 = InBgra0->B / 255.0f;
+		static_assert(std::is_same_v<ColorComponentType, uint8> || std::is_same_v<ColorComponentType, unsigned short>, "Unsupported color component type");
+
+		float R0, G0, B0, R1, G1, B1;
 
-		const float R1 = InBgra1->R / 255.0f;
-		const float G1 = InBgra1->G / 255.0f;
-		const float B1 = InBgra1->B / 255.0f;
+		if constexpr (std::is_same_v<ColorComponentType, uint8>)
+		{
+			R0 = InputColor0->R / 255.f;
+			G0 = InputColor0->G / 255.f;
+			B0 = InputColor0->B / 255.f;
 
-		// 709 conversion
+			R1 = InputColor1->R / 255.f;
+			G1 = InputColor1->G / 255.f;
+			B1 = InputColor1->B / 255.f;
+		}
+		else
+		{
+			// ColorComponentType == unsigned short
+			
+			R0 = InputColor0->R.GetFloat();
+			G0 = InputColor0->G.GetFloat();
+			B0 = InputColor0->B.GetFloat();
+
+			R1 = InputColor1->R.GetFloat();
+			G1 = InputColor1->G.GetFloat();
+			B1 = InputColor1->B.GetFloat();
+		}
+
+		// Rec. 709 conversion
+		// See: https://en.wikipedia.org/wiki/YCbCr#ITU-R_BT.709_conversion
 		const float Yfull0 = R0 * 0.212639f + G0 * 0.7151687f + B0 * 0.0721932f;
 		const float Yfull1 = R1 * 0.212639f + G1 * 0.7151687f + B1 * 0.0721932f;
 		const float CbFull0 = R0 * (-0.1145922f) + G0 * (-0.3854078f) + B0 * 0.5f;
@@ -42,19 +67,63 @@ namespace AvidDNx
 		const float CbAvg = (CbFull0 + CbFull1) / 2.0f;
 		const float CrAvg = (CrFull0 + CrFull1) / 2.0f;
 
-		// video range conversion
-		const uint8 YVideoRange0 = (uint8)FMath::RoundToInt((219 * Yfull0 + 16));
-		const uint8 YVideoRange1 = (uint8)FMath::RoundToInt((219 * Yfull1 + 16));
-		const uint8 CbVideoRange = (uint8)FMath::RoundToInt((224 * CbAvg + 128));
-		const uint8 CrVideoRange = (uint8)FMath::RoundToInt((224 * CrAvg + 128));
-
-		return
+		// Video range conversion. 8-bit values are from the Rec. 709 specification. The 16-bit values were derived from this.
+		// Example: WhitePoint_Y_16bit = 235/256 * 65536, where 256 is the max value of a uint8, and 65536 for uint16.
+		constexpr uint8 WhitePoint_Y_8bit = 235;
+		constexpr uint8 WhitePoint_CbCr_8bit = 240;
+		constexpr uint8 BlackPoint_8bit = 16;
+		constexpr uint8 Midpoint_8bit = 128;
+		
+		constexpr uint16 WhitePoint_Y_16bit = 60160;
+		constexpr uint16 WhitePoint_CbCr_16bit = 61440;
+		constexpr uint16 BlackPoint_16bit = 4096;
+		constexpr uint16 Midpoint_16bit = 32768;
+
+		ColorComponentType RangeDifference_Y;
+		ColorComponentType RangeDifference_CbCr;
+		ColorComponentType BlackPoint;
+		ColorComponentType Midpoint;
+		if constexpr (std::is_same_v<ColorComponentType, uint8>)
 		{
-			YVideoRange0,
-			CbVideoRange,
-			YVideoRange1,
-			CrVideoRange
-		};
+			RangeDifference_Y = WhitePoint_Y_8bit - BlackPoint_8bit;
+			RangeDifference_CbCr = WhitePoint_CbCr_8bit - BlackPoint_8bit;
+			BlackPoint = BlackPoint_8bit;
+			Midpoint = Midpoint_8bit;
+		}
+		else if constexpr (std::is_same_v<ColorComponentType, unsigned short>)
+		{
+			RangeDifference_Y = WhitePoint_Y_16bit - BlackPoint_16bit;
+			RangeDifference_CbCr = WhitePoint_CbCr_16bit - BlackPoint_16bit;
+			BlackPoint = BlackPoint_16bit;
+			Midpoint = Midpoint_16bit;
+		}
+		
+		const ColorComponentType YVideoRange0 = static_cast<ColorComponentType>(FMath::RoundToInt((RangeDifference_Y * Yfull0 + BlackPoint)));
+		const ColorComponentType YVideoRange1 = static_cast<ColorComponentType>(FMath::RoundToInt((RangeDifference_Y * Yfull1 + BlackPoint)));
+		const ColorComponentType CbVideoRange = static_cast<ColorComponentType>(FMath::RoundToInt((RangeDifference_CbCr * CbAvg + Midpoint)));
+		const ColorComponentType CrVideoRange = static_cast<ColorComponentType>(FMath::RoundToInt((RangeDifference_CbCr * CrAvg + Midpoint)));
+
+		return ColorContainer(YVideoRange0, CbVideoRange, YVideoRange1, CrVideoRange);
+	}
+
+	/** Converts a FFloat16Color to Rec. 709 and video range. */
+	FAvidDNxEncoder::FRGB_16bit RGBtoRec709(const FFloat16Color* InColor)
+	{
+		constexpr unsigned short WhitePoint = 60160;
+		constexpr unsigned short BlackPoint = 4096;
+		constexpr unsigned short Difference = WhitePoint - BlackPoint;
+		constexpr float GammaExponent = 1.f / 2.4f;
+
+		// Rec. 709 conversion
+		const float R = FMath::Pow(InColor->R, GammaExponent);
+		const float G = FMath::Pow(InColor->G, GammaExponent);
+		const float B = FMath::Pow(InColor->B, GammaExponent);
+
+		// Convert to video range
+		return FAvidDNxEncoder::FRGB_16bit(
+			static_cast<unsigned short>(R * Difference + BlackPoint),
+			static_cast<unsigned short>(G * Difference + BlackPoint),
+			static_cast<unsigned short>(B * Difference + BlackPoint));
 	}
 
 	int32 GCD(int32 A, int32 B)
@@ -90,6 +159,8 @@ FAvidDNxEncoder::FAvidDNxEncoder(const FAvidDNxEncoderOptions& InOptions)
 	: Options(InOptions)
 	, bInitialized(false)
 	, bFinalized(false)
+	, WriteStartTimeSeconds(0)
+	, WriteEndTimeSeconds(0)
 	, DNxHRencoder(nullptr)
 	, DNxUncEncoder(nullptr)
 	, MXFwriter(nullptr)
@@ -142,13 +213,25 @@ bool FAvidDNxEncoder::Initialize()
 
 bool FAvidDNxEncoder::InitializeCompressedEncoder()
 {
+	const bool bIsRGB = (Options.Quality == EAvidDNxEncoderQuality::RGB444_12bit);
+	
+	DNX_ComponentType_t ComponentType = DNX_CT_UCHAR;
+	if (Options.Quality == EAvidDNxEncoderQuality::HQX_10bit)
+	{
+		ComponentType = DNX_CT_USHORT_10_6;
+	}
+	else if (Options.Quality == EAvidDNxEncoderQuality::RGB444_12bit)
+	{
+		ComponentType = DNX_CT_USHORT_12_4;
+	}
+	
 	const DNX_UncompressedParams_t UncompressedParamsHR =
 	{
 		sizeof(DNX_UncompressedParams_t),
-		DNX_CT_UCHAR, // Component type
+		ComponentType, // Component type
 		DNX_CV_709, // Color volume
-		DNX_CF_YCbCr, // Color format
-		DNX_CCO_YCbYCr_NoA, // Component order
+		bIsRGB ? DNX_CF_RGB : DNX_CF_YCbCr, // Color format
+		bIsRGB ? DNX_CCO_RGB_NoA : DNX_CCO_YCbYCr_NoA, // Component order
 		DNX_BFO_Progressive, // Field order
 		DNX_RGT_Display, // Raster geometry type
 		0, // Interfield gap bytes
@@ -161,17 +244,27 @@ bool FAvidDNxEncoder::InitializeCompressedEncoder()
 		0 // Row bytes2
 	};
 
+	unsigned int BitDepth = 8;
+	if (Options.Quality == EAvidDNxEncoderQuality::HQX_10bit)
+	{
+		BitDepth = 10;
+	}
+	else if (Options.Quality == EAvidDNxEncoderQuality::RGB444_12bit)
+	{
+		BitDepth = 12;
+	}
+
 	const DNX_CompressedParams_t CompressedParamsHR =
 	{
 		sizeof(DNX_CompressedParams_t),
 		Options.Width,
 		Options.Height,
-		DNX_HQ_COMPRESSION_ID,
+		static_cast<DNX_CompressionID_t>(Options.Quality),
 		DNX_CV_709, // Color volume
-		DNX_CF_YCbCr, // Color format
+		bIsRGB ? DNX_CF_RGB : DNX_CF_YCbCr, // Color format
 		// Parameters below are used for RI only
-		DNX_SSC_422, // Sub-sampling
-		8, // bit-depth, is used only for RI compression IDs
+		bIsRGB ? DNX_SSC_444 : DNX_SSC_422, // Sub-sampling
+		BitDepth, // bit-depth, is used only for RI compression IDs
 		1, // PARC
 		1, // PARN
 		0, // CRC-presence
@@ -275,47 +368,21 @@ bool FAvidDNxEncoder::InitializeMXFWriter()
 	return true;
 }
 
-bool FAvidDNxEncoder::WriteFrame(const uint8* InFrameData)
+bool FAvidDNxEncoder::WriteFrame_Avid(const void* InSubSampledBuffer, const int32 InSubSampledBufferSize, void* OutEncodedBuffer, const int32 InEncodedBufferSize)
 {
-	const double StartTime = FPlatformTime::Seconds();
-
-	int32 NumPixels = Options.Width * Options.Height;
-
-	TArray<FY0CbY1Cr, TAlignedHeapAllocator<16>> SubSampledBuffer;
-	SubSampledBuffer.AddUninitialized(NumPixels / 2);
-	TArray<uint8, TAlignedHeapAllocator<16>> EncodedBuffer;
-	EncodedBuffer.AddUninitialized(EncodedBufferSize);
-	ensure(NumPixels % 2 == 0);
-
-
-	const FColor* ColorData = reinterpret_cast<const FColor*>(InFrameData);
-	for (size_t InIdx = 0, OutIdx = 0; InIdx < NumPixels; InIdx += 2, ++OutIdx)
-	{
-		// FAvidDNxEncoder::FY0CbY1Cr SubSample422(const FColor & InBgra0, const FColor & InBgra1)
-		const FColor* PixelA = &ColorData[InIdx];
-		const FColor* PixelB = &ColorData[InIdx + 1];
-
-		SubSampledBuffer[OutIdx] = AvidDNx::SubSample422(PixelA, PixelB);
-	}
-
-	const char* InBuffer = (const char*)SubSampledBuffer.GetData();
-	const unsigned int InputBufferSize = SubSampledBuffer.Num() * 4;
-
-	char* OutBuffer = (char*)EncodedBuffer.GetData();
-	const int32 OutBufferSize = EncodedBuffer.Num();
 	unsigned int CompressedFrameSize = 0;
 	bool bEncodingSuccessful = false;
 
-	const double ConvertionTime = FPlatformTime::Seconds();
+	const double ConversionTime = FPlatformTime::Seconds();
 
 	if (Options.bCompress)
 	{
 		const int EncodeStatus = DNX_EncodeFrame(
 			DNxHRencoder,
-			InBuffer,
-			OutBuffer,
-			InputBufferSize,
-			OutBufferSize,
+			InSubSampledBuffer,
+			OutEncodedBuffer,
+			InSubSampledBufferSize,
+			InEncodedBufferSize,
 			&CompressedFrameSize);
 
 		if (EncodeStatus != DNX_NO_ERROR)
@@ -337,24 +404,24 @@ bool FAvidDNxEncoder::WriteFrame(const uint8* InFrameData)
 			DNxUncEncoder,
 			&DNxUncUncompressedParams,
 			&DNxUncCompressedParams,
-			InBuffer,
+			InSubSampledBuffer,
 			UncompressedBufferSize,
-			OutBuffer,
-			OutBufferSize,
+			OutEncodedBuffer,
+			InEncodedBufferSize,
 			&CompressedFrameSize);
 		bEncodingSuccessful = EncodeStatus == DNX_UNCOMPRESSED_ERR_SUCCESS; // errors will be logged through DNxUncompressedErrorHandler
 	}
 
 	if (bEncodingSuccessful)
 	{
-		DNXMXF_WriteFrame(MXFwriter, OutBuffer, CompressedFrameSize);
+		DNXMXF_WriteFrame(MXFwriter, OutEncodedBuffer, CompressedFrameSize);
 	}
 
-	const double EndTime = FPlatformTime::Seconds();
+	WriteEndTimeSeconds = FPlatformTime::Seconds();
 
-	const double ConversionDeltaTimeMs = (ConvertionTime - StartTime) * 1000.0;
-	const double CodecDeltaTimeMs = (EndTime - ConvertionTime) * 1000.0;
-	const double TotalDeltaTimeMs = (EndTime - StartTime) * 1000.0;
+	const double ConversionDeltaTimeMs = (ConversionTime - WriteStartTimeSeconds) * 1000.0;
+	const double CodecDeltaTimeMs = (WriteEndTimeSeconds - ConversionTime) * 1000.0;
+	const double TotalDeltaTimeMs = (WriteEndTimeSeconds - WriteStartTimeSeconds) * 1000.0;
 
 	// FVideoFrameData* FramePayload = InFrame.GetPayload<FVideoFrameData>();
 	// UE_LOG(LogAvidDNxMedia, Verbose, TEXT("Processing Frame:%dx%d Frame:%d Conversion:%fms Codec:%fms Total:%fms"),
@@ -364,9 +431,115 @@ bool FAvidDNxEncoder::WriteFrame(const uint8* InFrameData)
 	// 	ConversionDeltaTimeMs,
 	// 	CodecDeltaTimeMs,
 	// 	TotalDeltaTimeMs);
+	
 	return bEncodingSuccessful;
 }
 
+bool FAvidDNxEncoder::WriteFrame(const uint8* InFrameData)
+{
+	WriteStartTimeSeconds = FPlatformTime::Seconds();
+	
+	const int32 NumPixels = Options.Width * Options.Height;
+
+	TArray<FY0CbY1Cr, TAlignedHeapAllocator<16>> SubSampledBuffer;
+	SubSampledBuffer.AddUninitialized(NumPixels / 2);
+	
+	TArray<uint8, TAlignedHeapAllocator<16>> EncodedBuffer;
+	EncodedBuffer.AddUninitialized(EncodedBufferSize);
+	
+	ensure(NumPixels % 2 == 0);
+
+	const FColor* ColorData = reinterpret_cast<const FColor*>(InFrameData);
+	ParallelFor(NumPixels / 2, [&SubSampledBuffer, &ColorData](const int32 SubSampledPixelIndex)
+	{
+		// The sub-sampled index goes from 0 -> NumPixels/2.
+		// The input index goes from 0 -> NumPixels.
+		
+		const int32 InputIndex = SubSampledPixelIndex * 2;
+		
+		const FColor* PixelA = &ColorData[InputIndex];
+		const FColor* PixelB = &ColorData[InputIndex + 1];
+
+		SubSampledBuffer[SubSampledPixelIndex] = AvidDNx::RBGtoYCbCrRec709<FColor, uint8, FY0CbY1Cr>(PixelA, PixelB);
+	});
+	
+	const char* InBuffer = reinterpret_cast<const char*>(SubSampledBuffer.GetData());
+	const int32 InputBufferSize = SubSampledBuffer.Num() * 4;
+
+	char* OutBuffer = reinterpret_cast<char*>(EncodedBuffer.GetData());
+	const int32 OutBufferSize = EncodedBuffer.Num();
+
+	return WriteFrame_Avid(InBuffer, InputBufferSize, OutBuffer, OutBufferSize);
+}
+
+bool FAvidDNxEncoder::WriteFrame_16bit(const FFloat16Color* InFrameData)
+{
+	const bool bIsRGB = (Options.Quality == EAvidDNxEncoderQuality::RGB444_12bit);
+
+	if (bIsRGB)
+	{
+		return WriteFrame_16bit_Impl<FRGB_16bit>(InFrameData);
+	}
+
+	return WriteFrame_16bit_Impl<FY0CbY1Cr_16bit>(InFrameData);
+}
+
+template <typename EncodedBufferType>
+bool FAvidDNxEncoder::WriteFrame_16bit_Impl(const FFloat16Color* InFrameData)
+{
+	WriteStartTimeSeconds = FPlatformTime::Seconds();
+	
+	const int32 NumPixels = Options.Width * Options.Height;
+	ensure(NumPixels % 2 == 0);
+
+	const bool bIs444 = (Options.Quality == EAvidDNxEncoderQuality::RGB444_12bit);
+	const bool bIsRGB = (Options.Quality == EAvidDNxEncoderQuality::RGB444_12bit);
+
+	// The buffer that the encoder will fill in
+	TArray<uint8, TAlignedHeapAllocator<16>> EncodedBuffer;
+	EncodedBuffer.AddUninitialized(EncodedBufferSize);
+
+	// The buffer provided to the encoder. Depending on the format, this buffer may not be sub-sampled at all (eg, RGB 4:4:4).
+	TArray<EncodedBufferType, TAlignedHeapAllocator<16>> SubSampledBuffer;
+	SubSampledBuffer.AddUninitialized(bIs444 ? NumPixels : NumPixels / 2);
+
+	// RGB -> Rec. 709 RGB
+	if constexpr (std::is_same_v<EncodedBufferType, FRGB_16bit>)
+	{
+		ParallelFor(NumPixels, [&SubSampledBuffer, &InFrameData](const int32 PixelIndex)
+		{
+			SubSampledBuffer[PixelIndex] = AvidDNx::RGBtoRec709(&InFrameData[PixelIndex]);
+		});
+	}
+
+	// RGB -> Rec. 709 YCbCr
+	else
+	{
+		ParallelFor(NumPixels / 2, [&SubSampledBuffer, &InFrameData](const int32 SubSampledPixelIndex)
+		{
+			// The sub-sampled index goes from 0 -> NumPixels/2.
+			// The input index goes from 0 -> NumPixels.
+			
+			const int32 InputIndex = SubSampledPixelIndex * 2;
+			
+			// Sub-sample the pixel data
+			const FFloat16Color* PixelA = &InFrameData[InputIndex];
+			const FFloat16Color* PixelB = &InFrameData[InputIndex + 1];
+
+			// Convert to video range and Rec 709
+			SubSampledBuffer[SubSampledPixelIndex] = AvidDNx::RBGtoYCbCrRec709<FFloat16Color, unsigned short, EncodedBufferType>(PixelA, PixelB);
+		});
+	}
+
+	const char* InBuffer = reinterpret_cast<const char*>(SubSampledBuffer.GetData());
+	const int32 InputBufferSize = bIsRGB ? sizeof(FRGB_16bit) * NumPixels : sizeof(FY0CbY1Cr_16bit) * NumPixels;
+
+	void* OutBuffer = EncodedBuffer.GetData();
+	const int32 OutBufferSize = EncodedBuffer.Num();
+
+	return WriteFrame_Avid(InBuffer, InputBufferSize, OutBuffer, OutBufferSize);
+}
+
 void FAvidDNxEncoder::Finalize()
 {
 	if (bFinalized || !bInitialized)
diff --git a/Engine/Plugins/Media/AvidDNxMedia/Source/Source/Private/MovieGraphAvidDNxHRNode.cpp b/Engine/Plugins/Media/AvidDNxMedia/Source/Source/Private/MovieGraphAvidDNxHRNode.cpp
--- a/Engine/Plugins/Media/AvidDNxMedia/Source/Source/Private/MovieGraphAvidDNxHRNode.cpp
+++ b/Engine/Plugins/Media/AvidDNxMedia/Source/Source/Private/MovieGraphAvidDNxHRNode.cpp
@@ -50,12 +50,17 @@ FSlateIcon UMovieGraphAvidDNxHRNode::GetIconAndTint(FLinearColor& OutColor) cons
 	return AvidDNxHRIcon;
 }
 
-TUniquePtr<MovieRenderGraph::IVideoCodecWriter> UMovieGraphAvidDNxHRNode::Initialize_GameThread(UMovieGraphPipeline* InPipeline, TObjectPtr<UMovieGraphEvaluatedConfig> InEvaluatedConfig, const FString& InFileName, FIntPoint InResolution, EImagePixelType InPixelType, ERGBFormat InPixelFormat, uint8 InBitDepth, uint8 InNumChannels, bool bAllowOCIO)
+TUniquePtr<MovieRenderGraph::IVideoCodecWriter> UMovieGraphAvidDNxHRNode::Initialize_GameThread(UMovieGraphPipeline* InPipeline, TObjectPtr<UMovieGraphEvaluatedConfig> InEvaluatedConfig, const FString& InBranchName, const FString& InFileName, FIntPoint InResolution, EImagePixelType InPixelType, ERGBFormat InPixelFormat, uint8 InBitDepth, uint8 InNumChannels, bool bAllowOCIO)
 {
-	constexpr bool bIncludeCDOs = true;
+	bool bIncludeCDOs = true;
 	constexpr bool bExactMatch = true;
 	UMovieGraphGlobalOutputSettingNode* OutputSetting =
 		InEvaluatedConfig->GetSettingForBranch<UMovieGraphGlobalOutputSettingNode>(GlobalsPinName, bIncludeCDOs, bExactMatch);
+
+	bIncludeCDOs = false;
+	const UMovieGraphAvidDNxHRNode* EvaluatedNode = Cast<UMovieGraphAvidDNxHRNode>(
+		InEvaluatedConfig->GetSettingForBranch(GetClass(), FName(InBranchName), bIncludeCDOs, bExactMatch));
+	checkf(EvaluatedNode, TEXT("Avid DNxHR node could not be found in the graph in branch [%s]."), *InBranchName);
 	
 	const FFrameRate SourceFrameRate = InPipeline->GetDataSourceInstance()->GetDisplayRate();
 	const FFrameRate EffectiveFrameRate = UMovieGraphBlueprintLibrary::GetEffectiveFrameRate(OutputSetting, SourceFrameRate);
@@ -64,6 +69,7 @@ TUniquePtr<MovieRenderGraph::IVideoCodecWriter> UMovieGraphAvidDNxHRNode::Initia
 	Options.OutputFilename = InFileName;
 	Options.Width = InResolution.X;
 	Options.Height = InResolution.Y;
+	Options.Quality = EvaluatedNode->Quality;
 	Options.FrameRate = EffectiveFrameRate;
 	Options.bCompress = true;
 	Options.NumberOfEncodingThreads = 4;
@@ -72,8 +78,15 @@ TUniquePtr<MovieRenderGraph::IVideoCodecWriter> UMovieGraphAvidDNxHRNode::Initia
 	NewWriter->Writer = MakeUnique<FAvidDNxEncoder>(Options);
 	NewWriter->FileName = InFileName;
 
-	// If OCIO is enabled, don't do additional color conversion
-	NewWriter->bConvertToSrgb = !(bOverride_OCIOConfiguration && OCIOConfiguration.bIsEnabled && bAllowOCIO);
+	// If OCIO is enabled, don't do additional color conversion. RGB444 12-bit is never converted to sRGB.
+	if (EvaluatedNode->Quality == EAvidDNxEncoderQuality::RGB444_12bit)
+	{
+		NewWriter->bConvertToSrgb = false;
+	}
+	else
+	{
+		NewWriter->bConvertToSrgb = !(bOverride_OCIOConfiguration && OCIOConfiguration.bIsEnabled && bAllowOCIO);
+	}
 
 	CachedPipeline = InPipeline;
 	
@@ -98,8 +111,16 @@ void UMovieGraphAvidDNxHRNode::WriteFrame_EncodeThread(MovieRenderGraph::IVideoC
 	
 	const UE::MovieGraph::FMovieGraphSampleState* Payload = InPixelData->GetPayload<UE::MovieGraph::FMovieGraphSampleState>();
 
-	// Quantize our 16-bit float data to 8-bit and apply sRGB if needed
-	const TUniquePtr<FImagePixelData> QuantizedPixelData = UE::MoviePipeline::QuantizeImagePixelDataToBitDepth(InPixelData, 8, nullptr, InWriter->bConvertToSrgb);
+	constexpr bool bIncludeCDOs = false;
+	constexpr bool bExactMatch = true;
+	const FName BranchName = Payload->TraversalContext.RenderDataIdentifier.RootBranchName;
+	const UMovieGraphAvidDNxHRNode* EvaluatedNode = Cast<UMovieGraphAvidDNxHRNode>(
+		InEvaluatedConfig->GetSettingForBranch(GetClass(), BranchName, bIncludeCDOs, bExactMatch));
+	checkf(EvaluatedNode, TEXT("Avid DNxHR node could not be found in the graph in branch [%s]."), *BranchName.ToString());
+
+	// Quantize our 16-bit float data to 8/16-bit and apply sRGB if needed
+	const int32 BitDepth = ((EvaluatedNode->Quality == EAvidDNxEncoderQuality::HQX_10bit) || (EvaluatedNode->Quality == EAvidDNxEncoderQuality::RGB444_12bit)) ? 16 : 8;
+	TUniquePtr<FImagePixelData> QuantizedPixelData = UE::MoviePipeline::QuantizeImagePixelDataToBitDepth(InPixelData, BitDepth, nullptr, InWriter->bConvertToSrgb);
 
 	TArray<FPixelPreProcessor> PixelPreProcessors;
 
@@ -138,7 +159,14 @@ void UMovieGraphAvidDNxHRNode::WriteFrame_EncodeThread(MovieRenderGraph::IVideoC
 	int64 DataSize;
 	QuantizedPixelData->GetRawData(Data, DataSize);
 
-	CodecWriter->Writer->WriteFrame((uint8*)Data);
+	if (BitDepth == 8)
+	{
+		CodecWriter->Writer->WriteFrame((uint8*)Data);
+	}
+	else
+	{
+		CodecWriter->Writer->WriteFrame_16bit((FFloat16Color*)Data);
+	}
 }
 
 void UMovieGraphAvidDNxHRNode::BeginFinalize_EncodeThread(MovieRenderGraph::IVideoCodecWriter* InWriter)
diff --git a/Engine/Plugins/Media/AvidDNxMedia/Source/Source/Private/MovieGraphAvidDNxHRNode.h b/Engine/Plugins/Media/AvidDNxMedia/Source/Source/Private/MovieGraphAvidDNxHRNode.h
--- a/Engine/Plugins/Media/AvidDNxMedia/Source/Source/Private/MovieGraphAvidDNxHRNode.h
+++ b/Engine/Plugins/Media/AvidDNxMedia/Source/Source/Private/MovieGraphAvidDNxHRNode.h
@@ -2,6 +2,7 @@
 
 #pragma once
 
+#include "AvidDNxEncoder/AvidDNxEncoder.h"
 #include "Graph/MovieGraphNode.h"
 #include "Graph/Nodes/MovieGraphVideoOutputNode.h"
 
@@ -17,13 +18,6 @@ enum class EMovieGraphAvidDNxHRFormat : uint8
 	Mov UMETA(DisplayName = "QuickTime (MOV)")
 };
 
-/** The codecs available for use with the Avid DNxHR node. */
-UENUM(BlueprintType)
-enum class EMovieGraphAvidDNxHRCodec : uint8
-{
-	HQ8Bit UMETA(DisplayName = "DNxHR HQ 8-bit"),
-};
-
 /** A node which can output Avid DNxHR movies. */
 UCLASS(BlueprintType, PrioritizeCategories=("FileOutput"))
 class UMovieGraphAvidDNxHRNode : public UMovieGraphVideoOutputNode
@@ -45,7 +39,7 @@ public:
 
 protected:
 	// UMovieGraphVideoOutputNode Interface
-	virtual TUniquePtr<MovieRenderGraph::IVideoCodecWriter> Initialize_GameThread(UMovieGraphPipeline* InPipeline, TObjectPtr<UMovieGraphEvaluatedConfig> InEvaluatedConfig, const FString& InFileName, FIntPoint InResolution, EImagePixelType InPixelType, ERGBFormat InPixelFormat, uint8 InBitDepth, uint8 InNumChannels, bool bAllowOCIO) override;
+	virtual TUniquePtr<MovieRenderGraph::IVideoCodecWriter> Initialize_GameThread(UMovieGraphPipeline* InPipeline, TObjectPtr<UMovieGraphEvaluatedConfig> InEvaluatedConfig, const FString& InBranchName, const FString& InFileName, FIntPoint InResolution, EImagePixelType InPixelType, ERGBFormat InPixelFormat, uint8 InBitDepth, uint8 InNumChannels, bool bAllowOCIO) override;
 	virtual bool Initialize_EncodeThread(MovieRenderGraph::IVideoCodecWriter* InWriter) override;
 	virtual void WriteFrame_EncodeThread(MovieRenderGraph::IVideoCodecWriter* InWriter, FImagePixelData* InPixelData, TArray<FMovieGraphPassData>&& InCompositePasses, TObjectPtr<UMovieGraphEvaluatedConfig> InEvaluatedConfig) override;
 	virtual void BeginFinalize_EncodeThread(MovieRenderGraph::IVideoCodecWriter* InWriter) override;
@@ -68,7 +62,7 @@ public:
 	uint8 bOverride_Format : 1;
 
 	UPROPERTY(EditAnywhere, BlueprintReadWrite, Category = Overrides, meta = (InlineEditConditionToggle))
-	uint8 bOverride_Codec : 1;
+	uint8 bOverride_Quality : 1;
 
 	UPROPERTY(EditAnywhere, BlueprintReadWrite, Category = Overrides, meta = (InlineEditConditionToggle))
 	uint8 bOverride_OCIOConfiguration : 1;
@@ -80,9 +74,9 @@ public:
 	UPROPERTY(EditAnywhere, BlueprintReadWrite, Category="Avid DNxHR", meta=(EditCondition="bOverride_Format"))
 	EMovieGraphAvidDNxHRFormat Format;
 
-	/** Movie Render Queue currently supports the Avid DNxHR HQ 8-bit codec only. */
-	UPROPERTY(EditAnywhere, BlueprintReadWrite, Category="Avid DNxHR", meta=(EditCondition="bOverride_Codec"))
-	EMovieGraphAvidDNxHRCodec Codec;
+	/**  */
+	UPROPERTY(EditAnywhere, BlueprintReadWrite, Category="Avid DNxHR", meta=(EditCondition="bOverride_Quality"))
+	EAvidDNxEncoderQuality Quality = EAvidDNxEncoderQuality::HQ_8bit;
 
 	/**
 	* OCIO configuration/transform settings.
diff --git a/Engine/Plugins/Media/AvidDNxMedia/Source/Source/Public/AvidDNxEncoder/AvidDNxEncoder.h b/Engine/Plugins/Media/AvidDNxMedia/Source/Source/Public/AvidDNxEncoder/AvidDNxEncoder.h
--- a/Engine/Plugins/Media/AvidDNxMedia/Source/Source/Public/AvidDNxEncoder/AvidDNxEncoder.h
+++ b/Engine/Plugins/Media/AvidDNxMedia/Source/Source/Public/AvidDNxEncoder/AvidDNxEncoder.h
@@ -6,14 +6,37 @@
 #include "Misc/FrameRate.h"
 
 //THIRDPARTY_INCLUDES_START
+#include <AvidDNxCodec.h>
 #include <dnx_uncompressed_sdk.h>
 //THIRDPARTY_INCLUDES_END
 
+#include "AvidDNxEncoder.generated.h"
+
 // Forward Declares
 struct _DNX_Encoder;
 using DNX_Encoder = struct _DNX_Encoder*;
 using DNXMXF_Writer = struct DNXMXF_Writer;
 
+/** Quality settings available for the Avid DNx encoder. */
+UENUM()
+enum class EAvidDNxEncoderQuality
+{
+	/** Cinema quality, 12-bit 4:4:4 RGB. */
+	RGB444_12bit	= DNX_444_COMPRESSION_ID	UMETA(DisplayName = "DNxHR RGB 444 12-bit"),
+
+	/** High quality extended, 10-bit 4:2:2 YCbCr. */
+	HQX_10bit		= DNX_HQX_COMPRESSION_ID	UMETA(DisplayName = "DNxHR HQX 10-bit"),
+
+	/** High quality, 8-bit 4:2:2 YCbCr. */
+	HQ_8bit			= DNX_HQ_COMPRESSION_ID		UMETA(DisplayName = "DNxHR HQ 8-bit"),
+
+	/** Standard quality, 8-bit 4:2:2 YCbCr. */
+	SQ_8bit			= DNX_SQ_COMPRESSION_ID		UMETA(DisplayName = "DNxHR SQ 8-bit"),
+
+	/** Low bandwidth, 8-bit 4:2:2 YCbCr. */
+	LB_8bit			= DNX_LB_COMPRESSION_ID		UMETA(DisplayName = "DNxHR LB 8-bit")
+};
+
 /** 
 * Options to initialize the AvidDNx encoder with. Choosing compression will choose the AvidDNxHR HD compression.
 */
@@ -23,6 +46,7 @@ struct FAvidDNxEncoderOptions
 		: OutputFilename()
 		, Width(0)
 		, Height(0)
+		, Quality(EAvidDNxEncoderQuality::HQ_8bit)
 		, FrameRate(30, 1)
 		, bCompress(true)
 		, NumberOfEncodingThreads(0)
@@ -37,6 +61,9 @@ struct FAvidDNxEncoderOptions
 	/** The height of the video file. */
 	uint32 Height;
 
+	/** The quality setting that the encoder should use. */
+	EAvidDNxEncoderQuality Quality;
+
 	/** Frame Rate of the output video. */
 	FFrameRate FrameRate;
 	
@@ -63,23 +90,53 @@ public:
 	/** Finalize the video file and finish writing it to disk. Called by the destructor if not automatically called. */
 	void Finalize();
 
-	/** Appends a new frame onto the output file. */
+	/** Appends a new frame onto the output file (8-bit). */
 	bool WriteFrame(const uint8* InFrameData);
 
+	/** Appends a new frame onto the output file (16-bit). */
+	bool WriteFrame_16bit(const FFloat16Color* InFrameData);
+
 private:
 	bool InitializeCompressedEncoder();
 	bool InitializeUncompressedEncoder();
 	bool InitializeMXFWriter();
+
+	/** Ask the AVID SDK to encode and write the buffer. */
+	bool WriteFrame_Avid(const void* InSubSampledBuffer, int32 InSubSampledBufferSize, void* OutEncodedBuffer, int32 InEncodedBufferSize);
+
+	/** Templated implementation of WriteFrame_16bit(). EncodedBufferType supports FRGB_16bit and FY0CbY1Cr_16bit. */
+	template <typename EncodedBufferType>
+	bool WriteFrame_16bit_Impl(const FFloat16Color* InFrameData);
+
 public:
+	/** 8-bit sub-sampled YCbCr color. */
 	struct FY0CbY1Cr
 	{
 		uint8 Y0, Cb, Y1, Cr;
 	};
 
+	/** 16-bit sub-sampled YCbCr color. */
+	struct FY0CbY1Cr_16bit
+	{
+		unsigned short Y0, Cb, Y1, Cr;
+	};
+
+	/** 16-bit RGB color. */
+	struct FRGB_16bit
+	{
+		unsigned short R, G, B;
+	};
+
 private:
 	FAvidDNxEncoderOptions Options;
 	bool bInitialized;
 	bool bFinalized;
+
+	/** When encoding and writing a frame started. */
+	double WriteStartTimeSeconds;
+
+	/** When encoding and writing a frame finished. */
+	double WriteEndTimeSeconds;
 	
 	/** How big each video sample is after compression based on given settings. */
 	int32 EncodedBufferSize;
diff --git a/Engine/Plugins/MovieScene/MovieRenderPipeline/Source/MovieRenderPipelineCore/Private/Graph/Nodes/MovieGraphVideoOutputNode.cpp b/Engine/Plugins/MovieScene/MovieRenderPipeline/Source/MovieRenderPipelineCore/Private/Graph/Nodes/MovieGraphVideoOutputNode.cpp
--- a/Engine/Plugins/MovieScene/MovieRenderPipeline/Source/MovieRenderPipelineCore/Private/Graph/Nodes/MovieGraphVideoOutputNode.cpp
+++ b/Engine/Plugins/MovieScene/MovieRenderPipeline/Source/MovieRenderPipelineCore/Private/Graph/Nodes/MovieGraphVideoOutputNode.cpp
@@ -254,7 +254,7 @@ UMovieGraphVideoOutputNode::FMovieGraphCodecWriterWithPromise* UMovieGraphVideoO
 	{
 		// Create a new writer for this file name (and output format settings)
 		TUniquePtr<MovieRenderGraph::IVideoCodecWriter> NewWriter = Initialize_GameThread(
-			InPipeline, EvaluatedConfig, FinalFilePath,
+			InPipeline, EvaluatedConfig, InRenderPassData.Key.RootBranchName.ToString(), FinalFilePath,
 			InRenderPassData.Value->GetSize(), InRenderPassData.Value->GetType(), InRenderPassData.Value->GetPixelLayout(),
 			InRenderPassData.Value->GetBitDepth(), InRenderPassData.Value->GetNumChannels(), Payload->bAllowOCIO);
 
diff --git a/Engine/Plugins/MovieScene/MovieRenderPipeline/Source/MovieRenderPipelineCore/Public/Graph/Nodes/MovieGraphVideoOutputNode.h b/Engine/Plugins/MovieScene/MovieRenderPipeline/Source/MovieRenderPipelineCore/Public/Graph/Nodes/MovieGraphVideoOutputNode.h
--- a/Engine/Plugins/MovieScene/MovieRenderPipeline/Source/MovieRenderPipelineCore/Public/Graph/Nodes/MovieGraphVideoOutputNode.h
+++ b/Engine/Plugins/MovieScene/MovieRenderPipeline/Source/MovieRenderPipelineCore/Public/Graph/Nodes/MovieGraphVideoOutputNode.h
@@ -49,7 +49,7 @@ protected:
 	virtual bool IsFinishedWritingToDiskImpl() const override;
 	// ~UMovieGraphFileOutputNode Interface
 
-	virtual TUniquePtr<MovieRenderGraph::IVideoCodecWriter> Initialize_GameThread(UMovieGraphPipeline* InPipeline, TObjectPtr<UMovieGraphEvaluatedConfig> InEvaluatedConfig, const FString& InFileName, FIntPoint InResolution, EImagePixelType InPixelType, ERGBFormat InPixelFormat, uint8 InBitDepth, uint8 InNumChannels, bool bAllowOCIO)  PURE_VIRTUAL(UMovieGraphVideoOutputNode::Initialize_GameThread, return nullptr; );
+	virtual TUniquePtr<MovieRenderGraph::IVideoCodecWriter> Initialize_GameThread(UMovieGraphPipeline* InPipeline, TObjectPtr<UMovieGraphEvaluatedConfig> InEvaluatedConfig, const FString& InBranchName, const FString& InFileName, FIntPoint InResolution, EImagePixelType InPixelType, ERGBFormat InPixelFormat, uint8 InBitDepth, uint8 InNumChannels, bool bAllowOCIO)  PURE_VIRTUAL(UMovieGraphVideoOutputNode::Initialize_GameThread, return nullptr; );
 	virtual bool Initialize_EncodeThread(MovieRenderGraph::IVideoCodecWriter* InWriter) PURE_VIRTUAL(UMovieGraphVideoOutputNode::Initialize_EncodeThread, return true;);
 	virtual void WriteFrame_EncodeThread(MovieRenderGraph::IVideoCodecWriter* InWriter, FImagePixelData* InPixelData, TArray<FMovieGraphPassData>&& InCompositePasses, TObjectPtr<UMovieGraphEvaluatedConfig> InEvaluatedConfig) PURE_VIRTUAL(UMovieGraphVideoOutputNode::WriteFrame_EncodeThread);
 	virtual void BeginFinalize_EncodeThread(MovieRenderGraph::IVideoCodecWriter* InWriter) PURE_VIRTUAL(UMovieGraphVideoOutputNode::BeginFinalize_EncodeThread);
