MoviePipeline: Updated the Avid DNxHR Movie node to include timecode in the generated movie.

diff --git a/Engine/Plugins/Media/AppleProResMedia/Source/AppleProResMedia/Private/MovieGraphAppleProResNode.cpp b/Engine/Plugins/Media/AppleProResMedia/Source/AppleProResMedia/Private/MovieGraphAppleProResNode.cpp
--- a/Engine/Plugins/Media/AppleProResMedia/Source/AppleProResMedia/Private/MovieGraphAppleProResNode.cpp
+++ b/Engine/Plugins/Media/AppleProResMedia/Source/AppleProResMedia/Private/MovieGraphAppleProResNode.cpp
@@ -58,25 +58,25 @@ FSlateIcon UMovieGraphAppleProResNode::GetIconAndTint(FLinearColor& OutColor) co
 }
 #endif // WITH_EDITOR
 
-TUniquePtr<MovieRenderGraph::IVideoCodecWriter> UMovieGraphAppleProResNode::Initialize_GameThread(UMovieGraphPipeline* InPipeline, TObjectPtr<UMovieGraphEvaluatedConfig> InEvaluatedConfig, const FString& InBranchName, const FString& InFileName, FIntPoint InResolution, EImagePixelType InPixelType, ERGBFormat InPixelFormat, uint8 InBitDepth, uint8 InNumChannels, bool bAllowOCIO)
+TUniquePtr<MovieRenderGraph::IVideoCodecWriter> UMovieGraphAppleProResNode::Initialize_GameThread(const FMovieGraphVideoNodeInitializationContext& InInitializationContext)
 {
 	bool bIncludeCDOs = true;
 	constexpr bool bExactMatch = true;
 	UMovieGraphGlobalOutputSettingNode* OutputSetting =
-		InEvaluatedConfig->GetSettingForBranch<UMovieGraphGlobalOutputSettingNode>(GlobalsPinName, bIncludeCDOs, bExactMatch);
+		InInitializationContext.EvaluatedConfig->GetSettingForBranch<UMovieGraphGlobalOutputSettingNode>(GlobalsPinName, bIncludeCDOs, bExactMatch);
 
 	bIncludeCDOs = false;
 	const UMovieGraphAppleProResNode* EvaluatedNode = Cast<UMovieGraphAppleProResNode>(
-		InEvaluatedConfig->GetSettingForBranch(GetClass(), FName(InBranchName), bIncludeCDOs, bExactMatch));
-	checkf(EvaluatedNode, TEXT("Apple ProRes node could not be found in the graph in branch [%s]."), *InBranchName);
+		InInitializationContext.EvaluatedConfig->GetSettingForBranch(GetClass(), FName(InInitializationContext.PassData->Key.RootBranchName), bIncludeCDOs, bExactMatch));
+	checkf(EvaluatedNode, TEXT("Apple ProRes node could not be found in the graph in branch [%s]."), *InInitializationContext.PassData->Key.RootBranchName.ToString());
 	
-	const FFrameRate SourceFrameRate = InPipeline->GetDataSourceInstance()->GetDisplayRate();
+	const FFrameRate SourceFrameRate = InInitializationContext.Pipeline->GetDataSourceInstance()->GetDisplayRate();
 	const FFrameRate EffectiveFrameRate = UMovieGraphBlueprintLibrary::GetEffectiveFrameRate(OutputSetting, SourceFrameRate);
 	
 	FAppleProResEncoderOptions Options;
-	Options.OutputFilename = InFileName;
-	Options.Width = InResolution.X;
-	Options.Height = InResolution.Y;
+	Options.OutputFilename = InInitializationContext.FileName;
+	Options.Width = InInitializationContext.PassData->Value->GetSize().X;
+	Options.Height = InInitializationContext.PassData->Value->GetSize().Y;
 	Options.FrameRate = EffectiveFrameRate;
 	Options.Codec = EvaluatedNode->Quality;
 	Options.ColorPrimaries = EAppleProResEncoderColorPrimaries::CD_HDREC709; // Force Rec 709 for now
@@ -85,12 +85,12 @@ TUniquePtr<MovieRenderGraph::IVideoCodecWriter> UMovieGraphAppleProResNode::Init
 	Options.bDropFrameTimecode = bDropFrameTimecode;
 
 	// If OCIO is enabled, don't do additional color conversion
-	Options.bConvertToSrgb = !(EvaluatedNode->bOverride_OCIOConfiguration && EvaluatedNode->OCIOConfiguration.bIsEnabled && bAllowOCIO);
+	Options.bConvertToSrgb = !(EvaluatedNode->bOverride_OCIOConfiguration && EvaluatedNode->OCIOConfiguration.bIsEnabled && InInitializationContext.bAllowOCIO);
 	
 	TUniquePtr<FProResWriter> NewWriter = MakeUnique<FProResWriter>();
 	NewWriter->Writer = MakeUnique<FAppleProResEncoder>(Options);
 
-	CachedPipeline = InPipeline;
+	CachedPipeline = InInitializationContext.Pipeline;
 	
 	return NewWriter;
 }
diff --git a/Engine/Plugins/Media/AppleProResMedia/Source/AppleProResMedia/Private/MovieGraphAppleProResNode.h b/Engine/Plugins/Media/AppleProResMedia/Source/AppleProResMedia/Private/MovieGraphAppleProResNode.h
--- a/Engine/Plugins/Media/AppleProResMedia/Source/AppleProResMedia/Private/MovieGraphAppleProResNode.h
+++ b/Engine/Plugins/Media/AppleProResMedia/Source/AppleProResMedia/Private/MovieGraphAppleProResNode.h
@@ -29,7 +29,7 @@ public:
 
 protected:
 	// UMovieGraphVideoOutputNode Interface
-	virtual TUniquePtr<MovieRenderGraph::IVideoCodecWriter> Initialize_GameThread(UMovieGraphPipeline* InPipeline, TObjectPtr<UMovieGraphEvaluatedConfig> InEvaluatedConfig, const FString& InBranchName, const FString& InFileName, FIntPoint InResolution, EImagePixelType InPixelType, ERGBFormat InPixelFormat, uint8 InBitDepth, uint8 InNumChannels, bool bAllowOCIO) override;
+	virtual TUniquePtr<MovieRenderGraph::IVideoCodecWriter> Initialize_GameThread(const FMovieGraphVideoNodeInitializationContext& InInitializationContext) override;
 	virtual bool Initialize_EncodeThread(MovieRenderGraph::IVideoCodecWriter* InWriter) override;
 	virtual void WriteFrame_EncodeThread(MovieRenderGraph::IVideoCodecWriter* InWriter, FImagePixelData* InPixelData, TArray<FMovieGraphPassData>&& InCompositePasses, TObjectPtr<UMovieGraphEvaluatedConfig> InEvaluatedConfig, const FString& InBranchName) override;
 	virtual void BeginFinalize_EncodeThread(MovieRenderGraph::IVideoCodecWriter* InWriter) override;
diff --git a/Engine/Plugins/Media/AvidDNxMedia/Source/Source/Private/AvidDNxEncoder/AvidDNxEncoder.cpp b/Engine/Plugins/Media/AvidDNxMedia/Source/Source/Private/AvidDNxEncoder/AvidDNxEncoder.cpp
--- a/Engine/Plugins/Media/AvidDNxMedia/Source/Source/Private/AvidDNxEncoder/AvidDNxEncoder.cpp
+++ b/Engine/Plugins/Media/AvidDNxMedia/Source/Source/Private/AvidDNxEncoder/AvidDNxEncoder.cpp
@@ -342,6 +342,15 @@ bool FAvidDNxEncoder::InitializeMXFWriter()
 	};
 	const DNXMXF_Rational_t AspectRatio = AvidDNx::AspectRatioFromResolution(Options.Width, Options.Height);
 
+	const FFrameRate TwentyNineNineSeven = FFrameRate(30000, 1001);
+	DNXMXF_TimeCodeComponent_t TimeCodeComponent = DNXMXF_TimeCodeComponent_t(
+		Options.StartTimecode.Hours,
+		Options.StartTimecode.Minutes,
+		Options.StartTimecode.Seconds,
+		Options.StartTimecode.Frames,
+		Options.StartTimecode.bDropFrameFormat && (Options.FrameRate == TwentyNineNineSeven)
+	);
+
 	const DNXMXF_WriterParams_t MXFwriterParams
 	{
 		sizeof(DNXMXF_WriterParams_t),
@@ -356,7 +365,8 @@ bool FAvidDNxEncoder::InitializeMXFWriter()
 		nullptr,
 		AspectRatio,
 		0, 0,
-		Options.bCompress ? DNXMXF_ESSENCE_DNXHR_HD : DNXMXF_ESSENCE_DNXUNCOMPRESSED
+		Options.bCompress ? DNXMXF_ESSENCE_DNXHR_HD : DNXMXF_ESSENCE_DNXUNCOMPRESSED,
+		&TimeCodeComponent
 	};
 
 	if (DNXMXF_CreateWriter(&MXFoptions, &MXFwriterParams, &MXFwriter) != DNXMXF_SUCCESS)
diff --git a/Engine/Plugins/Media/AvidDNxMedia/Source/Source/Private/MovieGraphAvidDNxHRNode.cpp b/Engine/Plugins/Media/AvidDNxMedia/Source/Source/Private/MovieGraphAvidDNxHRNode.cpp
--- a/Engine/Plugins/Media/AvidDNxMedia/Source/Source/Private/MovieGraphAvidDNxHRNode.cpp
+++ b/Engine/Plugins/Media/AvidDNxMedia/Source/Source/Private/MovieGraphAvidDNxHRNode.cpp
@@ -13,6 +13,11 @@
 #include "MovieRenderPipelineCoreModule.h"
 #include "Styling/AppStyle.h"
 
+UMovieGraphAvidDNxHRNode::UMovieGraphAvidDNxHRNode()
+	: bDropFrameTimecode(true)
+{
+}
+
 EMovieGraphBranchRestriction UMovieGraphAvidDNxHRNode::GetBranchRestriction() const
 {
 	return EMovieGraphBranchRestriction::Globals;
@@ -50,29 +55,49 @@ FSlateIcon UMovieGraphAvidDNxHRNode::GetIconAndTint(FLinearColor& OutColor) cons
 	return AvidDNxHRIcon;
 }
 
-TUniquePtr<MovieRenderGraph::IVideoCodecWriter> UMovieGraphAvidDNxHRNode::Initialize_GameThread(UMovieGraphPipeline* InPipeline, TObjectPtr<UMovieGraphEvaluatedConfig> InEvaluatedConfig, const FString& InBranchName, const FString& InFileName, FIntPoint InResolution, EImagePixelType InPixelType, ERGBFormat InPixelFormat, uint8 InBitDepth, uint8 InNumChannels, bool bAllowOCIO)
+TUniquePtr<MovieRenderGraph::IVideoCodecWriter> UMovieGraphAvidDNxHRNode::Initialize_GameThread(const FMovieGraphVideoNodeInitializationContext& InInitializationContext)
 {
 	bool bIncludeCDOs = true;
 	constexpr bool bExactMatch = true;
 	UMovieGraphGlobalOutputSettingNode* OutputSetting =
-		InEvaluatedConfig->GetSettingForBranch<UMovieGraphGlobalOutputSettingNode>(GlobalsPinName, bIncludeCDOs, bExactMatch);
+		InInitializationContext.EvaluatedConfig->GetSettingForBranch<UMovieGraphGlobalOutputSettingNode>(GlobalsPinName, bIncludeCDOs, bExactMatch);
 
 	bIncludeCDOs = false;
 	const UMovieGraphAvidDNxHRNode* EvaluatedNode = Cast<UMovieGraphAvidDNxHRNode>(
-		InEvaluatedConfig->GetSettingForBranch(GetClass(), FName(InBranchName), bIncludeCDOs, bExactMatch));
-	checkf(EvaluatedNode, TEXT("Avid DNxHR node could not be found in the graph in branch [%s]."), *InBranchName);
+		InInitializationContext.EvaluatedConfig->GetSettingForBranch(GetClass(), InInitializationContext.PassData->Key.RootBranchName, bIncludeCDOs, bExactMatch));
+	checkf(EvaluatedNode, TEXT("Avid DNxHR node could not be found in the graph in branch [%s]."), *InInitializationContext.PassData->Key.RootBranchName.ToString());
 	
-	const FFrameRate SourceFrameRate = InPipeline->GetDataSourceInstance()->GetDisplayRate();
+	const FFrameRate SourceFrameRate = InInitializationContext.Pipeline->GetDataSourceInstance()->GetDisplayRate();
 	const FFrameRate EffectiveFrameRate = UMovieGraphBlueprintLibrary::GetEffectiveFrameRate(OutputSetting, SourceFrameRate);
+
+	// Determine the timecode that the movie should be started at
+	FTimecode StartTimecode;
+	if (EvaluatedNode->bOverride_CustomTimecodeStart)
+	{
+		const int32 OutputFrameNumber = InInitializationContext.TraversalContext->Time.OutputFrameNumber;
+
+		// When using a custom timecode start, just use the root-level frame number (relative to zero) offset by the custom timecode start
+		StartTimecode = FTimecode::FromFrameNumber(
+			OutputFrameNumber + EvaluatedNode->CustomTimecodeStart.ToFrameNumber(EffectiveFrameRate).Value,
+			EffectiveFrameRate,
+			EvaluatedNode->bDropFrameTimecode);
+	}
+	else
+	{
+		// This is the frame number on the global time, can have overlaps (between encoders) or repeats when using handle frames/slowmo.
+		StartTimecode = InInitializationContext.TraversalContext->Time.RootTimeCode;
+	}
 	
 	FAvidDNxEncoderOptions Options;
-	Options.OutputFilename = InFileName;
-	Options.Width = InResolution.X;
-	Options.Height = InResolution.Y;
+	Options.OutputFilename = InInitializationContext.FileName;
+	Options.Width = InInitializationContext.PassData->Value->GetSize().X;
+	Options.Height = InInitializationContext.PassData->Value->GetSize().Y;
 	Options.Quality = EvaluatedNode->Quality;
 	Options.FrameRate = EffectiveFrameRate;
 	Options.bCompress = true;
 	Options.NumberOfEncodingThreads = 4;
+	Options.bDropFrameTimecode = EvaluatedNode->bDropFrameTimecode;
+	Options.StartTimecode = StartTimecode;
 
 	// If OCIO is enabled, don't do additional color conversion. RGB444 12-bit is never converted to sRGB.
 	if (EvaluatedNode->Quality == EAvidDNxEncoderQuality::RGB444_12bit)
@@ -81,13 +106,13 @@ TUniquePtr<MovieRenderGraph::IVideoCodecWriter> UMovieGraphAvidDNxHRNode::Initia
 	}
 	else
 	{
-		Options.bConvertToSrgb = !(bOverride_OCIOConfiguration && OCIOConfiguration.bIsEnabled && bAllowOCIO);
+		Options.bConvertToSrgb = !(bOverride_OCIOConfiguration && OCIOConfiguration.bIsEnabled && InInitializationContext.bAllowOCIO);
 	}
 	
 	TUniquePtr<FAvidWriter> NewWriter = MakeUnique<FAvidWriter>();
 	NewWriter->Writer = MakeUnique<FAvidDNxEncoder>(Options);
 
-	CachedPipeline = InPipeline;
+	CachedPipeline = InInitializationContext.Pipeline;
 	
 	return NewWriter;
 }
diff --git a/Engine/Plugins/Media/AvidDNxMedia/Source/Source/Private/MovieGraphAvidDNxHRNode.h b/Engine/Plugins/Media/AvidDNxMedia/Source/Source/Private/MovieGraphAvidDNxHRNode.h
--- a/Engine/Plugins/Media/AvidDNxMedia/Source/Source/Private/MovieGraphAvidDNxHRNode.h
+++ b/Engine/Plugins/Media/AvidDNxMedia/Source/Source/Private/MovieGraphAvidDNxHRNode.h
@@ -25,7 +25,7 @@ class UMovieGraphAvidDNxHRNode : public UMovieGraphVideoOutputNode
 	GENERATED_BODY()
 
 public:
-	UMovieGraphAvidDNxHRNode() = default;
+	UMovieGraphAvidDNxHRNode();
 
 	virtual EMovieGraphBranchRestriction GetBranchRestriction() const override;
 
@@ -39,7 +39,7 @@ public:
 
 protected:
 	// UMovieGraphVideoOutputNode Interface
-	virtual TUniquePtr<MovieRenderGraph::IVideoCodecWriter> Initialize_GameThread(UMovieGraphPipeline* InPipeline, TObjectPtr<UMovieGraphEvaluatedConfig> InEvaluatedConfig, const FString& InBranchName, const FString& InFileName, FIntPoint InResolution, EImagePixelType InPixelType, ERGBFormat InPixelFormat, uint8 InBitDepth, uint8 InNumChannels, bool bAllowOCIO) override;
+	virtual TUniquePtr<MovieRenderGraph::IVideoCodecWriter> Initialize_GameThread(const FMovieGraphVideoNodeInitializationContext& InInitializationContext) override;
 	virtual bool Initialize_EncodeThread(MovieRenderGraph::IVideoCodecWriter* InWriter) override;
 	virtual void WriteFrame_EncodeThread(MovieRenderGraph::IVideoCodecWriter* InWriter, FImagePixelData* InPixelData, TArray<FMovieGraphPassData>&& InCompositePasses, TObjectPtr<UMovieGraphEvaluatedConfig> InEvaluatedConfig, const FString& InBranchName) override;
 	virtual void BeginFinalize_EncodeThread(MovieRenderGraph::IVideoCodecWriter* InWriter) override;
@@ -63,6 +63,12 @@ public:
 
 	UPROPERTY(EditAnywhere, BlueprintReadWrite, Category = Overrides, meta = (InlineEditConditionToggle))
 	uint8 bOverride_Quality : 1;
+	
+	UPROPERTY(EditAnywhere, BlueprintReadWrite, Category = Overrides, meta = (InlineEditConditionToggle))
+	uint8 bOverride_CustomTimecodeStart : 1;
+
+	UPROPERTY(EditAnywhere, BlueprintReadWrite, Category = Overrides, meta = (InlineEditConditionToggle))
+	uint8 bOverride_bDropFrameTimecode : 1;
 
 	UPROPERTY(EditAnywhere, BlueprintReadWrite, Category = Overrides, meta = (InlineEditConditionToggle))
 	uint8 bOverride_OCIOConfiguration : 1;
@@ -78,6 +84,14 @@ public:
 	UPROPERTY(EditAnywhere, BlueprintReadWrite, Category="Avid DNxHR", meta=(EditCondition="bOverride_Quality"))
 	EAvidDNxEncoderQuality Quality = EAvidDNxEncoderQuality::HQ_8bit;
 
+	/** Start the timecode at a specific value, rather than the value coming from the Level Sequence. */
+	UPROPERTY(EditAnywhere, BlueprintReadWrite, Category = "Avid DNxHR", meta = (EditCondition = "bOverride_CustomTimecodeStart"))
+	FTimecode CustomTimecodeStart;
+
+	/** Whether the embedded timecode track should be written using drop-frame format. Only applicable if the sequence framerate is 29.97. */
+	UPROPERTY(EditAnywhere, BlueprintReadWrite, Category = "Avid DNxHR", DisplayName = "Use DF Timecode if 29.97 FPS", meta = (EditCondition = "bOverride_bDropFrameTimecode"))
+	bool bDropFrameTimecode;
+
 	/**
 	* OCIO configuration/transform settings.
 	*
diff --git a/Engine/Plugins/Media/AvidDNxMedia/Source/Source/Public/AvidDNxEncoder/AvidDNxEncoder.h b/Engine/Plugins/Media/AvidDNxMedia/Source/Source/Public/AvidDNxEncoder/AvidDNxEncoder.h
--- a/Engine/Plugins/Media/AvidDNxMedia/Source/Source/Public/AvidDNxEncoder/AvidDNxEncoder.h
+++ b/Engine/Plugins/Media/AvidDNxMedia/Source/Source/Public/AvidDNxEncoder/AvidDNxEncoder.h
@@ -4,6 +4,7 @@
 
 #include "CoreMinimal.h"
 #include "Misc/FrameRate.h"
+#include "Misc/Timecode.h"
 
 //THIRDPARTY_INCLUDES_START
 #include <AvidDNxCodec.h>
@@ -50,6 +51,7 @@ struct FAvidDNxEncoderOptions
 		, FrameRate(30, 1)
 		, bCompress(true)
 		, NumberOfEncodingThreads(0)
+		, bDropFrameTimecode(false)
 	{}
 
 	/** The absolute path on disk to try and save the video file to. */
@@ -75,6 +77,12 @@ struct FAvidDNxEncoderOptions
 	
 	/** Number of Encoding Threads. Must be at least 1. */
 	uint32 NumberOfEncodingThreads;
+
+	/** If true, timecode track will use drop frame notation for the 29.97 frame rate. */
+	bool bDropFrameTimecode;
+
+	/** The timecode to start the movie at. */
+	FTimecode StartTimecode;
 };
 
 /**
diff --git a/Engine/Plugins/MovieScene/MovieRenderPipeline/Source/MovieRenderPipelineCore/Private/Graph/Nodes/MovieGraphVideoOutputNode.cpp b/Engine/Plugins/MovieScene/MovieRenderPipeline/Source/MovieRenderPipelineCore/Private/Graph/Nodes/MovieGraphVideoOutputNode.cpp
--- a/Engine/Plugins/MovieScene/MovieRenderPipeline/Source/MovieRenderPipelineCore/Private/Graph/Nodes/MovieGraphVideoOutputNode.cpp
+++ b/Engine/Plugins/MovieScene/MovieRenderPipeline/Source/MovieRenderPipelineCore/Private/Graph/Nodes/MovieGraphVideoOutputNode.cpp
@@ -255,13 +255,16 @@ UMovieGraphVideoOutputNode::FMovieGraphCodecWriterWithPromise* UMovieGraphVideoO
 	
 	if (!OutputWriter)
 	{
-		// Create a new writer for this file name (and output format settings)
-		TUniquePtr<MovieRenderGraph::IVideoCodecWriter> NewWriter = Initialize_GameThread(
-			InPipeline, EvaluatedConfig, InRenderPassData.Key.RootBranchName.ToString(), FinalFilePath,
-			InRenderPassData.Value->GetSize(), InRenderPassData.Value->GetType(), InRenderPassData.Value->GetPixelLayout(),
-			InRenderPassData.Value->GetBitDepth(), InRenderPassData.Value->GetNumChannels(), Payload->bAllowOCIO);
-
-		if (NewWriter)
+		FMovieGraphVideoNodeInitializationContext InitializationContext;
+		InitializationContext.Pipeline = InPipeline;
+		InitializationContext.EvaluatedConfig = EvaluatedConfig;
+		InitializationContext.TraversalContext = &InRawFrameData->TraversalContext;
+		InitializationContext.PassData = &InRenderPassData;
+		InitializationContext.FileName = FinalFilePath;
+		InitializationContext.bAllowOCIO = Payload->bAllowOCIO;
+
+		// Create a new writer for this file name (and output settings)
+		if (TUniquePtr<MovieRenderGraph::IVideoCodecWriter> NewWriter = Initialize_GameThread(InitializationContext))
 		{
 			// Store the stable filename this was generated with so we can match them up later.
 			NewWriter->StableFileName = StableFilePath;
diff --git a/Engine/Plugins/MovieScene/MovieRenderPipeline/Source/MovieRenderPipelineCore/Public/Graph/Nodes/MovieGraphVideoOutputNode.h b/Engine/Plugins/MovieScene/MovieRenderPipeline/Source/MovieRenderPipelineCore/Public/Graph/Nodes/MovieGraphVideoOutputNode.h
--- a/Engine/Plugins/MovieScene/MovieRenderPipeline/Source/MovieRenderPipelineCore/Public/Graph/Nodes/MovieGraphVideoOutputNode.h
+++ b/Engine/Plugins/MovieScene/MovieRenderPipeline/Source/MovieRenderPipelineCore/Public/Graph/Nodes/MovieGraphVideoOutputNode.h
@@ -35,6 +35,17 @@ public:
 	UMovieGraphVideoOutputNode();
 
 protected:
+	// The parameters to Initialize_GameThread() have changed a lot -- using a struct as a the only parameter will make future changes easier.
+	struct FMovieGraphVideoNodeInitializationContext
+	{
+		UMovieGraphPipeline* Pipeline;
+		TObjectPtr<UMovieGraphEvaluatedConfig> EvaluatedConfig;
+		const FMovieGraphTraversalContext* TraversalContext;
+		const FMovieGraphPassData* PassData;
+		FString FileName;
+		bool bAllowOCIO;
+	};
+	
 	// UMovieGraphFileOutputNode Interface
 	virtual void OnReceiveImageDataImpl(UMovieGraphPipeline* InPipeline, UE::MovieGraph::FMovieGraphOutputMergerFrame* InRawFrameData, const TSet<FMovieGraphRenderDataIdentifier>& InMask) override;
 	virtual void OnAllFramesSubmittedImpl(UMovieGraphPipeline* InPipeline, TObjectPtr<UMovieGraphEvaluatedConfig>& InPrimaryJobEvaluatedGraph) override;
@@ -43,7 +54,7 @@ protected:
 	virtual bool IsFinishedWritingToDiskImpl() const override;
 	// ~UMovieGraphFileOutputNode Interface
 
-	virtual TUniquePtr<MovieRenderGraph::IVideoCodecWriter> Initialize_GameThread(UMovieGraphPipeline* InPipeline, TObjectPtr<UMovieGraphEvaluatedConfig> InEvaluatedConfig, const FString& InBranchName, const FString& InFileName, FIntPoint InResolution, EImagePixelType InPixelType, ERGBFormat InPixelFormat, uint8 InBitDepth, uint8 InNumChannels, bool bAllowOCIO)  PURE_VIRTUAL(UMovieGraphVideoOutputNode::Initialize_GameThread, return nullptr; );
+	virtual TUniquePtr<MovieRenderGraph::IVideoCodecWriter> Initialize_GameThread(const FMovieGraphVideoNodeInitializationContext& InInitializationContext)  PURE_VIRTUAL(UMovieGraphVideoOutputNode::Initialize_GameThread, return nullptr; );
 	virtual bool Initialize_EncodeThread(MovieRenderGraph::IVideoCodecWriter* InWriter) PURE_VIRTUAL(UMovieGraphVideoOutputNode::Initialize_EncodeThread, return true;);
 	virtual void WriteFrame_EncodeThread(MovieRenderGraph::IVideoCodecWriter* InWriter, FImagePixelData* InPixelData, TArray<FMovieGraphPassData>&& InCompositePasses, TObjectPtr<UMovieGraphEvaluatedConfig> InEvaluatedConfig, const FString& InBranchName) PURE_VIRTUAL(UMovieGraphVideoOutputNode::WriteFrame_EncodeThread);
 	virtual void BeginFinalize_EncodeThread(MovieRenderGraph::IVideoCodecWriter* InWriter) PURE_VIRTUAL(UMovieGraphVideoOutputNode::BeginFinalize_EncodeThread);
