MoviePipeline: Removed the Format option for the Avid DNxHR node in the graph (the MXF format is the only format currently available).

#rnx

diff --git a/Engine/Plugins/Media/AvidDNxMedia/Source/Source/Private/MovieGraphAvidDNxHRNode.h b/Engine/Plugins/Media/AvidDNxMedia/Source/Source/Private/MovieGraphAvidDNxHRNode.h
--- a/Engine/Plugins/Media/AvidDNxMedia/Source/Source/Private/MovieGraphAvidDNxHRNode.h
+++ b/Engine/Plugins/Media/AvidDNxMedia/Source/Source/Private/MovieGraphAvidDNxHRNode.h
@@ -10,14 +10,6 @@
 
 class FAvidDNxEncoder;
 
-/** The container formats available for use with the Avid DNxHR node. */
-UENUM(BlueprintType)
-enum class EMovieGraphAvidDNxHRFormat : uint8
-{
-	Mxf UMETA(DisplayName = "Material Exchange Format (MXF)"),
-	Mov UMETA(DisplayName = "QuickTime (MOV)")
-};
-
 /** A node which can output Avid DNxHR movies. */
 UCLASS(BlueprintType, PrioritizeCategories=("FileOutput"))
 class UMovieGraphAvidDNxHRNode : public UMovieGraphVideoOutputNode
@@ -56,9 +48,6 @@ protected:
 	TWeakObjectPtr<UMovieGraphPipeline> CachedPipeline;
 
 public:
-	UPROPERTY(EditAnywhere, BlueprintReadWrite, Category = Overrides, meta = (InlineEditConditionToggle))
-	uint8 bOverride_Format : 1;
-
 	UPROPERTY(EditAnywhere, BlueprintReadWrite, Category = Overrides, meta = (InlineEditConditionToggle))
 	uint8 bOverride_Quality : 1;
 	
@@ -74,11 +63,7 @@ public:
 	UPROPERTY(EditAnywhere, BlueprintReadWrite, Category = Overrides, meta = (InlineEditConditionToggle))
 	uint8 bOverride_OCIOContext : 1;
 
-	/** The format to output the movie to. */
-	UPROPERTY(EditAnywhere, BlueprintReadWrite, Category="Avid DNxHR", meta=(EditCondition="bOverride_Format"))
-	EMovieGraphAvidDNxHRFormat Format;
-
-	/**  */
+	/** The quality that the movie will be encoded with. */
 	UPROPERTY(EditAnywhere, BlueprintReadWrite, Category="Avid DNxHR", meta=(EditCondition="bOverride_Quality"))
 	EAvidDNxEncoderQuality Quality = EAvidDNxEncoderQuality::HQ_8bit;
 
