MoviePipeline: Fixed some color issues in the Avid DNxHR Movie node.

Colors are now clamped to prevent some visual oddities, and the RGB 444 12bit preset applies a more correct Rec. 709 transform.

#rnx

diff --git a/Engine/Plugins/Media/AvidDNxMedia/Source/Source/Private/AvidDNxEncoder/AvidDNxEncoder.cpp b/Engine/Plugins/Media/AvidDNxMedia/Source/Source/Private/AvidDNxEncoder/AvidDNxEncoder.cpp
--- a/Engine/Plugins/Media/AvidDNxMedia/Source/Source/Private/AvidDNxEncoder/AvidDNxEncoder.cpp
+++ b/Engine/Plugins/Media/AvidDNxMedia/Source/Source/Private/AvidDNxEncoder/AvidDNxEncoder.cpp
@@ -34,25 +34,25 @@ namespace AvidDNx
 
 		if constexpr (std::is_same_v<ColorComponentType, uint8>)
 		{
-			R0 = InputColor0->R / 255.f;
-			G0 = InputColor0->G / 255.f;
-			B0 = InputColor0->B / 255.f;
+			R0 = FMath::Clamp(InputColor0->R / 255.f, 0.f, 1.f);
+			G0 = FMath::Clamp(InputColor0->G / 255.f, 0.f, 1.f);
+			B0 = FMath::Clamp(InputColor0->B / 255.f, 0.f, 1.f);
 
-			R1 = InputColor1->R / 255.f;
-			G1 = InputColor1->G / 255.f;
-			B1 = InputColor1->B / 255.f;
+			R1 = FMath::Clamp(InputColor1->R / 255.f, 0.f, 1.f);
+			G1 = FMath::Clamp(InputColor1->G / 255.f, 0.f, 1.f);
+			B1 = FMath::Clamp(InputColor1->B / 255.f, 0.f, 1.f);
 		}
 		else
 		{
 			// ColorComponentType == unsigned short
 			
-			R0 = InputColor0->R.GetFloat();
-			G0 = InputColor0->G.GetFloat();
-			B0 = InputColor0->B.GetFloat();
+			R0 = FMath::Clamp(InputColor0->R.GetFloat(), 0.f, 1.f);
+			G0 = FMath::Clamp(InputColor0->G.GetFloat(), 0.f, 1.f);
+			B0 = FMath::Clamp(InputColor0->B.GetFloat(), 0.f, 1.f);
 
-			R1 = InputColor1->R.GetFloat();
-			G1 = InputColor1->G.GetFloat();
-			B1 = InputColor1->B.GetFloat();
+			R1 = FMath::Clamp(InputColor1->R.GetFloat(), 0.f, 1.f);
+			G1 = FMath::Clamp(InputColor1->G.GetFloat(), 0.f, 1.f);
+			B1 = FMath::Clamp(InputColor1->B.GetFloat(), 0.f, 1.f);
 		}
 
 		// Rec. 709 conversion
@@ -115,9 +115,12 @@ namespace AvidDNx
 		constexpr float GammaExponent = 1.f / 2.4f;
 
 		// Rec. 709 conversion
-		const float R = FMath::Pow(InColor->R, GammaExponent);
-		const float G = FMath::Pow(InColor->G, GammaExponent);
-		const float B = FMath::Pow(InColor->B, GammaExponent);
+		const float R_Clamped = FMath::Clamp(InColor->R, 0.f, 1.f);
+		const float G_Clamped = FMath::Clamp(InColor->G, 0.f, 1.f);
+		const float B_Clamped = FMath::Clamp(InColor->B, 0.f, 1.f);
+		const float R = (R_Clamped <= 0.0031308f) ? R_Clamped * 12.92f : 1.055f * FMath::Pow(R_Clamped, GammaExponent) - 0.055f;
+		const float G = (G_Clamped <= 0.0031308f) ? G_Clamped * 12.92f : 1.055f * FMath::Pow(G_Clamped, GammaExponent) - 0.055f;
+		const float B = (B_Clamped <= 0.0031308f) ? B_Clamped * 12.92f : 1.055f * FMath::Pow(B_Clamped, GammaExponent) - 0.055f;
 
 		// Convert to video range
 		return FAvidDNxEncoder::FRGB_16bit(
