MovieGraph: Fixed OCIO configurations not being applied to jpg/png files due to using the CDO of the node for the configuration values, instead of the actual evaluated graph.

This resulted in over-exposed/blown out images when OCIO was applied, regardless of OCIO settings.

diff --git a/Engine/Plugins/Media/AppleProResMedia/Source/AppleProResMedia/Private/MovieGraphAppleProResNode.cpp b/Engine/Plugins/Media/AppleProResMedia/Source/AppleProResMedia/Private/MovieGraphAppleProResNode.cpp
--- a/Engine/Plugins/Media/AppleProResMedia/Source/AppleProResMedia/Private/MovieGraphAppleProResNode.cpp
+++ b/Engine/Plugins/Media/AppleProResMedia/Source/AppleProResMedia/Private/MovieGraphAppleProResNode.cpp
@@ -141,7 +141,7 @@ void UMovieGraphAppleProResNode::WriteFrame_EncodeThread(MovieRenderGraph::IVide
 	TArray<FPixelPreProcessor> PixelPreProcessors;
 
 #if WITH_OCIO
-	FMovieGraphOCIOHelper::GenerateOcioPixelPreProcessor(GraphPayload, CachedPipeline.Get(), InEvaluatedConfig, OCIOConfiguration, OCIOContext, PixelPreProcessors);
+	FMovieGraphOCIOHelper::GenerateOcioPixelPreProcessor(GraphPayload, CachedPipeline.Get(), InEvaluatedConfig, EvaluatedNode->OCIOConfiguration, EvaluatedNode->OCIOContext, PixelPreProcessors);
 #endif
 
 	// Do a quick composite of renders/burn-ins.
diff --git a/Engine/Plugins/Media/AvidDNxMedia/Source/Source/Private/MovieGraphAvidDNxHRNode.cpp b/Engine/Plugins/Media/AvidDNxMedia/Source/Source/Private/MovieGraphAvidDNxHRNode.cpp
--- a/Engine/Plugins/Media/AvidDNxMedia/Source/Source/Private/MovieGraphAvidDNxHRNode.cpp
+++ b/Engine/Plugins/Media/AvidDNxMedia/Source/Source/Private/MovieGraphAvidDNxHRNode.cpp
@@ -1,4 +1,4 @@
-﻿// Copyright Epic Games, Inc. All Rights Reserved.
+// Copyright Epic Games, Inc. All Rights Reserved.
 
 #include "MovieGraphAvidDNxHRNode.h"
 
@@ -101,7 +101,7 @@ TUniquePtr<MovieRenderGraph::IVideoCodecWriter> UMovieGraphAvidDNxHRNode::Initia
 	}
 	else
 	{
-		Options.bConvertToSrgb = !(bOverride_OCIOConfiguration && OCIOConfiguration.bIsEnabled && InInitializationContext.bAllowOCIO);
+		Options.bConvertToSrgb = !(EvaluatedNode->bOverride_OCIOConfiguration && EvaluatedNode->OCIOConfiguration.bIsEnabled && InInitializationContext.bAllowOCIO);
 	}
 	
 	TUniquePtr<FAvidWriter> NewWriter = MakeUnique<FAvidWriter>();
@@ -144,7 +144,7 @@ void UMovieGraphAvidDNxHRNode::WriteFrame_EncodeThread(MovieRenderGraph::IVideoC
 	TArray<FPixelPreProcessor> PixelPreProcessors;
 
 #if WITH_OCIO
-	FMovieGraphOCIOHelper::GenerateOcioPixelPreProcessor(Payload, CachedPipeline.Get(), InEvaluatedConfig, OCIOConfiguration, OCIOContext, PixelPreProcessors);
+	FMovieGraphOCIOHelper::GenerateOcioPixelPreProcessor(Payload, CachedPipeline.Get(), InEvaluatedConfig, EvaluatedNode->OCIOConfiguration, EvaluatedNode->OCIOContext, PixelPreProcessors);
 #endif
 
 	// Do a quick composite of renders/burn-ins.
diff --git a/Engine/Plugins/MovieScene/MovieRenderPipeline/Source/MovieRenderPipelineRenderPasses/Private/MovieGraphImageSequenceOutputNode.cpp b/Engine/Plugins/MovieScene/MovieRenderPipeline/Source/MovieRenderPipelineRenderPasses/Private/MovieGraphImageSequenceOutputNode.cpp
--- a/Engine/Plugins/MovieScene/MovieRenderPipeline/Source/MovieRenderPipelineRenderPasses/Private/MovieGraphImageSequenceOutputNode.cpp
+++ b/Engine/Plugins/MovieScene/MovieRenderPipeline/Source/MovieRenderPipelineRenderPasses/Private/MovieGraphImageSequenceOutputNode.cpp
@@ -161,7 +161,7 @@ void UMovieGraphImageSequenceOutputNode::OnReceiveImageDataImpl(UMovieGraphPipel
 		
 		bool bQuantizationEncodeSRGB = true;
 #if WITH_OCIO
-		if (FMovieGraphOCIOHelper::GenerateOcioPixelPreProcessor(Payload, InPipeline, InRawFrameData->EvaluatedConfig.Get(), OCIOConfiguration, OCIOContext, TileImageTask->PixelPreProcessors))
+		if (FMovieGraphOCIOHelper::GenerateOcioPixelPreProcessor(Payload, InPipeline, InRawFrameData->EvaluatedConfig.Get(), ParentNode->OCIOConfiguration, ParentNode->OCIOContext, TileImageTask->PixelPreProcessors))
 		{
 			// We assume that any encoding on the output transform should be done by OCIO
 			bQuantizationEncodeSRGB = false;
