MoviePipeline: Fixed several OCIO issues across MRQ and MRG.

* The MRG MP4 node did a double-sRGB transform if using OCIO, resulting in a brighter image than expected.
* The same issue was happening in the ProRes node in MRQ (this is a regression introduced in 5.6.0).
* Across both MRQ and MRG, all video settings/nodes incorrectly applied quantization before the OCIO transform, which could sometimes result in incorrect colors, particularly in the shadows/highlights.

Also fixed an issue where MRG would sometimes composite multiple burn-ins on the same frame in videos (if using multiple layers), potentially resulting in darker burn-ins than expected.

diff --git a/Engine/Plugins/Media/AppleProResMedia/Source/AppleProResMedia/Private/MovieGraphAppleProResNode.cpp b/Engine/Plugins/Media/AppleProResMedia/Source/AppleProResMedia/Private/MovieGraphAppleProResNode.cpp
--- a/Engine/Plugins/Media/AppleProResMedia/Source/AppleProResMedia/Private/MovieGraphAppleProResNode.cpp
+++ b/Engine/Plugins/Media/AppleProResMedia/Source/AppleProResMedia/Private/MovieGraphAppleProResNode.cpp
@@ -144,16 +144,21 @@ void UMovieGraphAppleProResNode::WriteFrame_EncodeThread(MovieRenderGraph::IVide
 		ProResPayload->ReferenceFrameNumber = GraphPayload->TraversalContext.Time.RootFrameNumber.Value;
 	}
 
-	// ProRes can handle quantization internally but expects sRGB to be applied to the incoming data (Rec 709)
-	const bool bConvertToSrgb = !CodecWriter->bSkipColorConversions;
-	const TUniquePtr<FImagePixelData> QuantizedPixelData = UE::MoviePipeline::QuantizeImagePixelDataToBitDepth(InPixelData, 8, ProResPayload, bConvertToSrgb);
-
 	TArray<FPixelPreProcessor> PixelPreProcessors;
 
+	// Run OCIO before quantizing. For highest accuracy, OCIO should run on the non-quantized pixel data.
 #if WITH_OCIO
-	FMovieGraphOCIOHelper::GenerateOcioPixelPreProcessor(GraphPayload, CachedPipeline.Get(), InEvaluatedConfig, EvaluatedNode->OCIOConfiguration, EvaluatedNode->OCIOContext, PixelPreProcessors);
+	if (FMovieGraphOCIOHelper::GenerateOcioPixelPreProcessor(GraphPayload, CachedPipeline.Get(), InEvaluatedConfig, EvaluatedNode->OCIOConfiguration, EvaluatedNode->OCIOContext, PixelPreProcessors))
+	{
+		PixelPreProcessors.Pop()(InPixelData);
+	}
 #endif
 
+	// ProRes can handle quantization to 8bit internally if necessary; apply sRGB if needed.
+	constexpr int32 TargetBitDepth = 16;
+	const bool bConvertToSrgb = !CodecWriter->bSkipColorConversions;
+	const TUniquePtr<FImagePixelData> QuantizedPixelData = UE::MoviePipeline::QuantizeImagePixelDataToBitDepth(InPixelData, TargetBitDepth, ProResPayload, bConvertToSrgb);
+
 	// Do a quick composite of renders/burn-ins.
 	for (const FMovieGraphPassData& CompositePass : InCompositePasses)
 	{
diff --git a/Engine/Plugins/Media/AppleProResMedia/Source/AppleProResMedia/Private/MoviePipelineAppleProResOutput.cpp b/Engine/Plugins/Media/AppleProResMedia/Source/AppleProResMedia/Private/MoviePipelineAppleProResOutput.cpp
--- a/Engine/Plugins/Media/AppleProResMedia/Source/AppleProResMedia/Private/MoviePipelineAppleProResOutput.cpp
+++ b/Engine/Plugins/Media/AppleProResMedia/Source/AppleProResMedia/Private/MoviePipelineAppleProResOutput.cpp
@@ -68,8 +68,9 @@ void UMoviePipelineAppleProResOutput::WriteFrame_EncodeThread(MovieRenderPipelin
 	ProResPayload->ReferenceFrameNumber = PipelinePayload->SampleState.OutputState.SourceFrameNumber;
 
 	// ProRes can handle either 16 or 8 bit input internally, but expects Rec709 input which has a sRGB curve applied.
-	constexpr bool bConvertTosRGB = true;
-	TUniquePtr<FImagePixelData> QuantizedPixelData = UE::MoviePipeline::QuantizeImagePixelDataToBitDepth(InPixelData, 16, ProResPayload, bConvertTosRGB);
+	const bool bConvertTosRGB = InWriter->bConvertToSrgb;
+	constexpr int32 TargetBitDepth = 16;
+	const TUniquePtr<FImagePixelData> QuantizedPixelData = UE::MoviePipeline::QuantizeImagePixelDataToBitDepth(InPixelData, TargetBitDepth, ProResPayload, bConvertTosRGB);
 
 	// Do a quick composite of renders/burn-ins.
 	TArray<FPixelPreProcessor> PixelPreProcessors;
diff --git a/Engine/Plugins/Media/AvidDNxMedia/Source/Source/Private/MovieGraphAvidDNxHRNode.cpp b/Engine/Plugins/Media/AvidDNxMedia/Source/Source/Private/MovieGraphAvidDNxHRNode.cpp
--- a/Engine/Plugins/Media/AvidDNxMedia/Source/Source/Private/MovieGraphAvidDNxHRNode.cpp
+++ b/Engine/Plugins/Media/AvidDNxMedia/Source/Source/Private/MovieGraphAvidDNxHRNode.cpp
@@ -138,16 +138,20 @@ void UMovieGraphAvidDNxHRNode::WriteFrame_EncodeThread(MovieRenderGraph::IVideoC
 		InEvaluatedConfig->GetSettingForBranch(GetClass(), BranchName, bIncludeCDOs, bExactMatch));
 	checkf(EvaluatedNode, TEXT("Avid DNxHR node could not be found in the graph in branch [%s]."), *BranchName.ToString());
 
-	// Quantize our 16-bit float data to 8/16-bit and apply sRGB if needed
-	const int32 BitDepth = ((EvaluatedNode->Quality == EAvidDNxEncoderQuality::HQX_10bit) || (EvaluatedNode->Quality == EAvidDNxEncoderQuality::RGB444_12bit)) ? 16 : 8;
-	TUniquePtr<FImagePixelData> QuantizedPixelData = UE::MoviePipeline::QuantizeImagePixelDataToBitDepth(InPixelData, BitDepth, nullptr, CodecWriter->Writer->GetOptions().bConvertToSrgb);
-
 	TArray<FPixelPreProcessor> PixelPreProcessors;
 
+	// Run OCIO before quantizing. For highest accuracy, OCIO should run on the non-quantized pixel data.
 #if WITH_OCIO
-	FMovieGraphOCIOHelper::GenerateOcioPixelPreProcessor(Payload, CachedPipeline.Get(), InEvaluatedConfig, EvaluatedNode->OCIOConfiguration, EvaluatedNode->OCIOContext, PixelPreProcessors);
+	if (FMovieGraphOCIOHelper::GenerateOcioPixelPreProcessor(Payload, CachedPipeline.Get(), InEvaluatedConfig, EvaluatedNode->OCIOConfiguration, EvaluatedNode->OCIOContext, PixelPreProcessors))
+	{
+		PixelPreProcessors.Pop()(InPixelData);
+	}
 #endif
 
+	// Quantize our 16-bit float data to 8/16-bit and apply sRGB if needed
+	const int32 BitDepth = ((EvaluatedNode->Quality == EAvidDNxEncoderQuality::HQX_10bit) || (EvaluatedNode->Quality == EAvidDNxEncoderQuality::RGB444_12bit)) ? 16 : 8;
+	const TUniquePtr<FImagePixelData> QuantizedPixelData = UE::MoviePipeline::QuantizeImagePixelDataToBitDepth(InPixelData, BitDepth, nullptr, CodecWriter->Writer->GetOptions().bConvertToSrgb);
+
 	// Do a quick composite of renders/burn-ins.
 	for (const FMovieGraphPassData& CompositePass : InCompositePasses)
 	{
diff --git a/Engine/Plugins/MovieScene/MovieRenderPipeline/Source/MovieRenderPipelineCore/Private/Graph/Nodes/MovieGraphVideoOutputNode.cpp b/Engine/Plugins/MovieScene/MovieRenderPipeline/Source/MovieRenderPipelineCore/Private/Graph/Nodes/MovieGraphVideoOutputNode.cpp
--- a/Engine/Plugins/MovieScene/MovieRenderPipeline/Source/MovieRenderPipelineCore/Private/Graph/Nodes/MovieGraphVideoOutputNode.cpp
+++ b/Engine/Plugins/MovieScene/MovieRenderPipeline/Source/MovieRenderPipelineCore/Private/Graph/Nodes/MovieGraphVideoOutputNode.cpp
@@ -143,6 +143,12 @@ void UMovieGraphVideoOutputNode::OnReceiveImageDataImpl(UMovieGraphPipeline* InP
 				{
 					continue;
 				}
+
+				// There might be composites across multiple branches; pick out the composites specific to this branch.
+				if (RenderPassData.Key.RootBranchName != CompositePass.Key.RootBranchName)
+				{
+					continue;
+				}
 				
 				// Match them up by camera name so multiple passes intended for different camera names work.
 				if (RenderPassData.Key.CameraName == CompositePass.Key.CameraName)
diff --git a/Engine/Plugins/MovieScene/MovieRenderPipeline/Source/MovieRenderPipelineMP4Encoder/Private/Graph/MovieGraphMP4EncoderNode.cpp b/Engine/Plugins/MovieScene/MovieRenderPipeline/Source/MovieRenderPipelineMP4Encoder/Private/Graph/MovieGraphMP4EncoderNode.cpp
--- a/Engine/Plugins/MovieScene/MovieRenderPipeline/Source/MovieRenderPipelineMP4Encoder/Private/Graph/MovieGraphMP4EncoderNode.cpp
+++ b/Engine/Plugins/MovieScene/MovieRenderPipeline/Source/MovieRenderPipelineMP4Encoder/Private/Graph/MovieGraphMP4EncoderNode.cpp
@@ -134,15 +134,20 @@ void UMovieGraphMP4EncoderNode::WriteFrame_EncodeThread(MovieRenderGraph::IVideo
 	
 	const UE::MovieGraph::FMovieGraphSampleState* GraphPayload = InPixelData->GetPayload<UE::MovieGraph::FMovieGraphSampleState>();
 	
-	const bool bConvertToSrgb = !CodecWriter->bSkipColorConversions;
-	const TUniquePtr<FImagePixelData> QuantizedPixelData = UE::MoviePipeline::QuantizeImagePixelDataToBitDepth(InPixelData, 8, nullptr, bConvertToSrgb);
-
 	TArray<FPixelPreProcessor> PixelPreProcessors;
 
+	// Run OCIO before quantizing. For highest accuracy, OCIO should run on the non-quantized pixel data.
 #if WITH_OCIO
-	FMovieGraphOCIOHelper::GenerateOcioPixelPreProcessor(GraphPayload, CachedPipeline.Get(), InEvaluatedConfig, EvaluatedNode->OCIOConfiguration, EvaluatedNode->OCIOContext, PixelPreProcessors);
+	if (FMovieGraphOCIOHelper::GenerateOcioPixelPreProcessor(GraphPayload, CachedPipeline.Get(), InEvaluatedConfig, EvaluatedNode->OCIOConfiguration, EvaluatedNode->OCIOContext, PixelPreProcessors))
+	{
+		PixelPreProcessors.Pop()(InPixelData);
+	}
 #endif
 
+	constexpr int32 TargetBitDepth = 8;
+	const bool bConvertToSrgb = !CodecWriter->bSkipColorConversions;
+	const TUniquePtr<FImagePixelData> QuantizedPixelData = UE::MoviePipeline::QuantizeImagePixelDataToBitDepth(InPixelData, TargetBitDepth, nullptr, bConvertToSrgb);
+
 	// Do a quick composite of renders/burn-ins.
 	for (const FMovieGraphPassData& CompositePass : InCompositePasses)
 	{
diff --git a/Engine/Plugins/MovieScene/MovieRenderPipeline/Source/MovieRenderPipelineMP4Encoder/Private/MoviePipelineMP4EncoderOutput.cpp b/Engine/Plugins/MovieScene/MovieRenderPipeline/Source/MovieRenderPipelineMP4Encoder/Private/MoviePipelineMP4EncoderOutput.cpp
--- a/Engine/Plugins/MovieScene/MovieRenderPipeline/Source/MovieRenderPipelineMP4Encoder/Private/MoviePipelineMP4EncoderOutput.cpp
+++ b/Engine/Plugins/MovieScene/MovieRenderPipeline/Source/MovieRenderPipelineMP4Encoder/Private/MoviePipelineMP4EncoderOutput.cpp
@@ -61,7 +61,9 @@ void UMoviePipelineMP4EncoderOutput::WriteFrame_EncodeThread(MovieRenderPipeline
 	FMP4CodecWriter* CodecWriter = static_cast<FMP4CodecWriter*>(InWriter);
 	
 	// Quantize our 16 bit float data to 8 bit.
-	TUniquePtr<FImagePixelData> QuantizedPixelData = UE::MoviePipeline::QuantizeImagePixelDataToBitDepth(InPixelData, 8);
+	const bool bConvertTosRGB = InWriter->bConvertToSrgb;
+	constexpr int32 TargetBitDepth = 8;
+	const TUniquePtr<FImagePixelData> QuantizedPixelData = UE::MoviePipeline::QuantizeImagePixelDataToBitDepth(InPixelData, TargetBitDepth, nullptr, bConvertTosRGB);
 
 	// Do a quick composite of renders/burn-ins.
 	TArray<FPixelPreProcessor> PixelPreProcessors;
